Add GC reference tracking for Async Message Payload data.

#rnx

diff --git a/Engine/Plugins/Experimental/AsyncMessageSystem/Source/AsyncMessageSystem/Private/AsyncMessageStore.cpp b/Engine/Plugins/Experimental/AsyncMessageSystem/Source/AsyncMessageSystem/Private/AsyncMessageStore.cpp
--- a/Engine/Plugins/Experimental/AsyncMessageSystem/Source/AsyncMessageSystem/Private/AsyncMessageStore.cpp
+++ b/Engine/Plugins/Experimental/AsyncMessageSystem/Source/AsyncMessageSystem/Private/AsyncMessageStore.cpp
@@ -22,7 +22,7 @@ uint32 FAsyncMessageStore::EnqueueMessage(FAsyncMessage&& MessageToQueue, OUT TA
 			// Only add this message to a queue if we know that it has a listener for this particular message.
 			if (Pair.Value.MessagesWithTheseBindings.Contains(CurrentMessageId) && !BindingTypesAddedTo.Contains(Pair.Key))
 			{
-				Pair.Value.MessageQueue->Enqueue(MessageToQueue);
+				Pair.Value.MessageQueue->Add(MessageToQueue);
 				++NumQueuesAddedTo;
 			
 				BindingTypesAddedTo.Emplace(Pair.Key);
@@ -40,10 +40,10 @@ TOptional<FAsyncMessage> FAsyncMessageStore::GetNextMessageForBindingOption(cons
 
 	if (FBindingOptionsMessageQueue* QueueContainer = MessageBindingQueues.Find(Binding))
 	{
-		FAsyncMessage OutMessage {};
-		if (QueueContainer->MessageQueue->Dequeue(OutMessage))
-		{			
-			return OutMessage;
+		// If we can, then pop an element off the message queue
+		if (!QueueContainer->MessageQueue->IsEmpty())
+		{
+			return QueueContainer->MessageQueue->Pop(EAllowShrinking::No);
 		}
 	}
 	
@@ -87,4 +87,20 @@ bool FAsyncMessageStore::RemoveMessageFromBinding(const FAsyncMessageId& Message
 
 	// Otherwise there are still other references to this message bound to this binding, so don't remove it
 	return false;
+}
+
+void FAsyncMessageStore::AddReferencedObjects(FReferenceCollector& Collector)
+{
+	FScopeLock QueueLock(&MessageQueueCriticalSection);
+
+	// For each message in the queue, add any GC refs that its payload may have.
+	for (TPair<FAsyncMessageBindingOptions, FBindingOptionsMessageQueue>& Pair : MessageBindingQueues)
+	{
+		FBindingOptionsMessageQueue& QueueData = Pair.Value;
+		
+		for (FAsyncMessage& Message : *(QueueData.MessageQueue.Get()))
+		{
+			Message.AddReferencedObjects(Collector);
+		}
+	}
 }
\ No newline at end of file
diff --git a/Engine/Plugins/Experimental/AsyncMessageSystem/Source/AsyncMessageSystem/Private/AsyncMessageSystemBase.cpp b/Engine/Plugins/Experimental/AsyncMessageSystem/Source/AsyncMessageSystem/Private/AsyncMessageSystemBase.cpp
--- a/Engine/Plugins/Experimental/AsyncMessageSystem/Source/AsyncMessageSystem/Private/AsyncMessageSystemBase.cpp
+++ b/Engine/Plugins/Experimental/AsyncMessageSystem/Source/AsyncMessageSystem/Private/AsyncMessageSystemBase.cpp
@@ -72,6 +72,17 @@ void FAsyncMessageSystemBase::Shutdown()
 	Shutdown_Impl();
 }
 
+void FAsyncMessageSystemBase::AddReferencedObjects(UObject* InReferencer, FReferenceCollector& Collector)
+{
+	// Add references from the message store 
+	MessageStore.AddReferencedObjects(Collector);
+
+	// Note: We are explictly NOT adding references to the _listener_ objects in the BoundMessageListenerMap
+	// because the message system only allows you to bind to weak references. We do not want a message system
+	// to keep a refernce to a listener object alive in GC, and we handle their removal in the BindListener
+	// lambda functions already.
+}
+
 ///////////////////////////////////////////////////////////
 // Virtual Implementations
 
diff --git a/Engine/Plugins/Experimental/AsyncMessageSystem/Source/AsyncMessageSystem/Private/AsyncMessageWorldSubsystem.cpp b/Engine/Plugins/Experimental/AsyncMessageSystem/Source/AsyncMessageSystem/Private/AsyncMessageWorldSubsystem.cpp
--- a/Engine/Plugins/Experimental/AsyncMessageSystem/Source/AsyncMessageSystem/Private/AsyncMessageWorldSubsystem.cpp
+++ b/Engine/Plugins/Experimental/AsyncMessageSystem/Source/AsyncMessageSystem/Private/AsyncMessageWorldSubsystem.cpp
@@ -17,6 +17,18 @@ FAutoConsoleVariableRef CVar_AsyncMessageEnableWorldSubsystemForEditor(
 	TEXT("When set to true, create an async message world subsystem for editor worlds as well as gameplay ones."),
 	ECVF_Default);
 
+void UAsyncMessageWorldSubsystem::AddReferencedObjects(UObject* InThis, FReferenceCollector& Collector)
+{
+	Super::AddReferencedObjects(InThis, Collector);
+
+	UAsyncMessageWorldSubsystem* ThisPtr = CastChecked<UAsyncMessageWorldSubsystem>(InThis);
+	
+	if (ThisPtr->MessageSystem.IsValid())
+	{
+		ThisPtr->MessageSystem->AddReferencedObjects(InThis, Collector);
+	}
+}
+
 bool UAsyncMessageWorldSubsystem::ShouldCreateSubsystem(UObject* Outer) const
 {
 	if (!bAsyncMessageEnableWorldSubsystem)
diff --git a/Engine/Plugins/Experimental/AsyncMessageSystem/Source/AsyncMessageSystem/Public/AsyncMessage.h b/Engine/Plugins/Experimental/AsyncMessageSystem/Source/AsyncMessageSystem/Public/AsyncMessage.h
--- a/Engine/Plugins/Experimental/AsyncMessageSystem/Source/AsyncMessageSystem/Public/AsyncMessage.h
+++ b/Engine/Plugins/Experimental/AsyncMessageSystem/Source/AsyncMessageSystem/Public/AsyncMessage.h
@@ -90,6 +90,12 @@ public:
 		return PayloadCopy.GetPtr<TDataType>();
 	}
 
+	
+	FORCEINLINE void AddReferencedObjects(FReferenceCollector& Collector)
+	{
+		PayloadCopy.AddStructReferencedObjects(Collector);
+	}
+
 private:
 	
 	/**
diff --git a/Engine/Plugins/Experimental/AsyncMessageSystem/Source/AsyncMessageSystem/Public/AsyncMessageStore.h b/Engine/Plugins/Experimental/AsyncMessageSystem/Source/AsyncMessageSystem/Public/AsyncMessageStore.h
--- a/Engine/Plugins/Experimental/AsyncMessageSystem/Source/AsyncMessageSystem/Public/AsyncMessageStore.h
+++ b/Engine/Plugins/Experimental/AsyncMessageSystem/Source/AsyncMessageSystem/Public/AsyncMessageStore.h
@@ -38,6 +38,13 @@ public:
 	 */
 	bool RemoveMessageFromBinding(const FAsyncMessageId& MessageId, const FAsyncMessageBindingOptions& Binding);
 
+	/**
+	 * Adds references to objects stored in the message queue to the GC collection
+	 * This will include the FInstancedStruct payload data on the message queue
+	 * to prevent objects from being destroyed if they are referenced in a message paylod.
+	 */
+	void AddReferencedObjects(FReferenceCollector& Collector);
+	
 protected:
 
 	/**
@@ -51,7 +58,7 @@ protected:
 	{
 		FBindingOptionsMessageQueue()
 		{
-			MessageQueue = MakeShared<TQueue<FAsyncMessage>>();
+			MessageQueue = MakeShared<TArray<FAsyncMessage>>();
 		}
 
 		~FBindingOptionsMessageQueue()
@@ -63,9 +70,9 @@ protected:
 		// find which queues require which messages.
 		// The value of this map is how many times it has been bound.
 		TMap<FAsyncMessageId, uint32> MessagesWithTheseBindings;
-
+		
 		// The queue of messages for this binding option		
-		TSharedPtr<TQueue<FAsyncMessage>> MessageQueue;
+		TSharedPtr<TArray<FAsyncMessage>> MessageQueue;
 	};
 
 	/**
diff --git a/Engine/Plugins/Experimental/AsyncMessageSystem/Source/AsyncMessageSystem/Public/AsyncMessageSystemBase.h b/Engine/Plugins/Experimental/AsyncMessageSystem/Source/AsyncMessageSystem/Public/AsyncMessageSystemBase.h
--- a/Engine/Plugins/Experimental/AsyncMessageSystem/Source/AsyncMessageSystem/Public/AsyncMessageSystemBase.h
+++ b/Engine/Plugins/Experimental/AsyncMessageSystem/Source/AsyncMessageSystem/Public/AsyncMessageSystemBase.h
@@ -146,6 +146,8 @@ public:
 	 */
 	void Shutdown();
 
+	void AddReferencedObjects(class UObject* InReferencer, FReferenceCollector& Collector);
+
 protected:
 
 	/**
diff --git a/Engine/Plugins/Experimental/AsyncMessageSystem/Source/AsyncMessageSystem/Public/AsyncMessageWorldSubsystem.h b/Engine/Plugins/Experimental/AsyncMessageSystem/Source/AsyncMessageSystem/Public/AsyncMessageWorldSubsystem.h
--- a/Engine/Plugins/Experimental/AsyncMessageSystem/Source/AsyncMessageSystem/Public/AsyncMessageWorldSubsystem.h
+++ b/Engine/Plugins/Experimental/AsyncMessageSystem/Source/AsyncMessageSystem/Public/AsyncMessageWorldSubsystem.h
@@ -52,6 +52,8 @@ public:
 	 * to this subsystem's message system when the world goes out of scope
 	 */
 	TMulticastDelegate<void()> OnShutdownMessageSystem;
+
+	static void AddReferencedObjects(UObject* InThis, FReferenceCollector& Collector);
 	
 protected:
 
diff --git a/Engine/Plugins/Experimental/AsyncMessageSystemTests/Source/AsyncMessageSystemTests/Private/AsyncMessageSystemPerformanceTests.cpp b/Engine/Plugins/Experimental/AsyncMessageSystemTests/Source/AsyncMessageSystemTests/Private/AsyncMessageSystemPerformanceTests.cpp
--- a/Engine/Plugins/Experimental/AsyncMessageSystemTests/Source/AsyncMessageSystemTests/Private/AsyncMessageSystemPerformanceTests.cpp
+++ b/Engine/Plugins/Experimental/AsyncMessageSystemTests/Source/AsyncMessageSystemTests/Private/AsyncMessageSystemPerformanceTests.cpp
@@ -521,7 +521,7 @@ bool TMessageSystemPerformance_BroadcastSeveralMessagesMultithread::RunTest(cons
 		PendingTasks.Push(T);
 	};
 	
-	RunTestTicks(TEXT("AsyncMessage_BroadcastSeveralMessages"), TickLambda);
+	RunTestTicks(TEXT("AsyncMessage_BroadcastSeveralMessages_Multi"), TickLambda);
 
 	UE::Tasks::Wait(PendingTasks);
 
@@ -531,4 +531,95 @@ bool TMessageSystemPerformance_BroadcastSeveralMessagesMultithread::RunTest(cons
 	return bSuccess && !ReportAnyErrors();
 }
 
-#endif	// WITH_DEV_AUTOMATION_TESTS
\ No newline at end of file
+UE_DEFINE_GAMEPLAY_TAG_COMMENT(InternalTestTag_RefCollection, "AsyncMessages.Internal.test.ReferenceCollection", "A test gameplay tag utilized in the async message system unit tests to test reference collection")
+static const FAsyncMessageId MessageId_RefCollection = { InternalTestTag_RefCollection };
+
+// Test that the message system correctly keeps tracks of reference UPROPERTY's in it's payloads
+// so that they do not get garbage collected whilst still on the message queue
+IMPLEMENT_CUSTOM_SIMPLE_AUTOMATION_TEST(TMessageSystem_PayloadReferenceTest, UE::AsyncMessageSystem::FAsyncMessageSystemTestBase, "AsyncMessagePassing.ReferenceCollection", PerformanceTestFlags)
+bool TMessageSystem_PayloadReferenceTest::RunTest(const FString& Parameters)
+{
+	using namespace UE::AsyncMessageSystem;
+
+	bool bSuccess = true;
+
+	bSuccess &= RunAsyncMessageTestSetup(EAsyncMessagePerfTestFlags::All, /*actor count override*/15);
+
+	// Bind a message to our event for populating a test actor on the game thread
+	ForEachTestActor([&](AASyncMessagePerfTest* TestActor)
+	{
+		if (TSharedPtr<FAsyncGameplayMessageSystem> Sys = UAsyncMessageWorldSubsystem::GetSharedMessageSystem<FAsyncGameplayMessageSystem>(TestActor->GetWorld()))
+		{
+			FAsyncMessageBindingOptions BindOptions = {};
+			
+			// Add a lamda function which will do some simple floating point test work in response to "DoSomeFakeWorkMessageId"
+			const FAsyncMessageHandle Handle = Sys->BindListener(MessageId_RefCollection, [&](const FAsyncMessage& Message)
+			{						
+				if (const FTest_RefCollection_Payload* Data = Message.GetPayloadData<const FTest_RefCollection_Payload>())
+				{
+					// Test to make sure that the payload object is valid and has not been GC'd
+					TestTrue(TEXT("Test data object is a valid pointer"), Data->ObjPoint != nullptr);
+					TestTrue(TEXT("Test data object is a Valid LowLevel"), Data->ObjPoint->IsValidLowLevel());
+					TestFalse(TEXT("Test data object is a reachable object!"), Data->ObjPoint->IsUnreachable());
+
+					if (Data->ObjPoint)
+					{
+						UE_LOG(LogTemp, Log, TEXT("Test actor is valid!"));
+
+						if (Data->ObjPoint->IsUnreachable())
+						{
+							UE_LOG(LogTemp, Error, TEXT("Test actor is unreachable! ruh roh"));
+						}
+						// We crash here prior to any work
+						else if (UTestRefCollectionObject* AsyncObj = Cast<UTestRefCollectionObject>(Data->ObjPoint))
+						{
+							UE_LOG(LogTemp, Log, TEXT("Even the type info is there!"));
+						}							
+					}
+					else
+					{
+						UE_LOG(LogTemp, Warning, TEXT("We have no test actor pointer"));
+					}
+				}
+			}, BindOptions);
+
+			TestActor->BoundHandles.Add(Handle);
+		}
+	});
+
+	// A simple tick which will queue one message for broadcasting each frame.
+	auto TickLambda = [](TSharedPtr<::FAsyncGameplayMessageSystem> MessageSys, float DeltaTime, int32 TickNum) -> void
+	{
+		// On tick 1, create a UObject which would only be referenced by the message we queue to the system 
+		if (TickNum == 1)
+		{
+			{
+				UTestRefCollectionObject* CreatedObject = NewObject<UTestRefCollectionObject>(GetTransientPackage(), TEXT("TestObject_0"));
+
+				FTest_RefCollection_Payload RefCollectionPayload {};
+				RefCollectionPayload.ObjPoint = CreatedObject;
+				FConstStructView PayloadView = FConstStructView::Make<FTest_RefCollection_Payload>(RefCollectionPayload);
+
+				MessageSys->QueueMessageForBroadcast(MessageId_RefCollection, PayloadView);	
+			}
+			
+			// And then immediately run GC. The Object will be collected because it won't be tracked anywhere
+			CollectGarbage(GARBAGE_COLLECTION_KEEPFLAGS);
+
+			// The message system will then proceed to tick again, and we will see what happens in the lambda above. 
+			// What will happen by default here is that the payload data object will not be referenced by anything, so
+			// will be destroyed by GC. This isn't what we want. We want to keep track of that UPROPERTY TObjectPtr 
+			// property on the FInstancedStruct payload data, and while it is in the message system queue, keep the references
+			// so it doesn't get collected
+		}
+	};
+	
+	RunTestTicks(TEXT("AsyncMessage_ReferenceCollection"), TickLambda, /* only tick 10 times to make this go faster */ 10);
+
+	// Always reset test world
+	bSuccess &= DestroyTestWorld();
+
+	return bSuccess && !ReportAnyErrors();
+}
+
+#endif	// WITH_DEV_AUTOMATION_TESTS
diff --git a/Engine/Plugins/Experimental/AsyncMessageSystemTests/Source/AsyncMessageSystemTests/Private/AsyncMessageSystemPerformanceTests.h b/Engine/Plugins/Experimental/AsyncMessageSystemTests/Source/AsyncMessageSystemTests/Private/AsyncMessageSystemPerformanceTests.h
--- a/Engine/Plugins/Experimental/AsyncMessageSystemTests/Source/AsyncMessageSystemTests/Private/AsyncMessageSystemPerformanceTests.h
+++ b/Engine/Plugins/Experimental/AsyncMessageSystemTests/Source/AsyncMessageSystemTests/Private/AsyncMessageSystemPerformanceTests.h
@@ -5,6 +5,7 @@
 #include "AsyncMessageHandle.h"
 #include "EngineRuntimeTests.h"			// For AEngineTestTickActor
 #include "UObject/WeakObjectPtrTemplates.h"
+#include "UObject/Object.h"				// For UObject
 
 #include "AsyncMessageSystemPerformanceTests.generated.h"
 
@@ -49,4 +50,30 @@ public:
 	float MathLimit = 1.0f;
 
 	void DoSimpleTestWork();
+};
+
+
+USTRUCT(NotBlueprintType, meta=(Hidden))
+struct FTest_RefCollection_Payload
+{
+	GENERATED_BODY()
+
+	UPROPERTY()
+	TObjectPtr<UObject> ObjPoint = nullptr;
+	
+	UPROPERTY()
+	TWeakObjectPtr<UObject> WeakObjPointer = nullptr;
+};
+
+/**
+ * A test UObject type which we will use to test ref counting of payload data
+ */
+UCLASS(NotBlueprintable, NotBlueprintType)
+class UTestRefCollectionObject : public UObject
+{
+	GENERATED_BODY()
+public:
+
+	UPROPERTY()
+	int32 TestValue = 5;
 };
\ No newline at end of file
